Trace thread explicitly shut down during UELibrary shutdown, to prevent a hang later.

diff --git a/Engine/Source/Runtime/UELibrary/Private/UELibrary.cpp b/Engine/Source/Runtime/UELibrary/Private/UELibrary.cpp
--- a/Engine/Source/Runtime/UELibrary/Private/UELibrary.cpp
+++ b/Engine/Source/Runtime/UELibrary/Private/UELibrary.cpp
@@ -156,6 +156,7 @@ int UELibrary_Shutdown()
 	EngineExit();
 	LaunchWindowsShutdown();
 	UE::UELibrary::Private::LibraryIsInitialized = false;
+	UE::Trace::Shutdown();
 	return 0;
 }
 
