Windows platform extensions use the default Windows launch code.

#rnx

diff --git a/Engine/Source/Runtime/Core/Public/Windows/WindowsPlatform.h b/Engine/Source/Runtime/Core/Public/Windows/WindowsPlatform.h
--- a/Engine/Source/Runtime/Core/Public/Windows/WindowsPlatform.h
+++ b/Engine/Source/Runtime/Core/Public/Windows/WindowsPlatform.h
@@ -112,7 +112,6 @@ typedef FWindowsPlatformTypes FPlatformTypes;
 #define WINDOWS_USE_FEATURE_APPLICATIONMISC_CLASS			1
 #define WINDOWS_USE_FEATURE_PLATFORMMISC_CLASS				1
 #define WINDOWS_USE_FEATURE_PLATFORMHTTP_CLASS				1
-#define WINDOWS_USE_FEATURE_LAUNCH							1
 
 
 // Q: Why is there a __nop() before __debugbreak()?
diff --git a/Engine/Source/Runtime/Launch/Private/Windows/LaunchWindows.cpp b/Engine/Source/Runtime/Launch/Private/Windows/LaunchWindows.cpp
--- a/Engine/Source/Runtime/Launch/Private/Windows/LaunchWindows.cpp
+++ b/Engine/Source/Runtime/Launch/Private/Windows/LaunchWindows.cpp
@@ -2,8 +2,6 @@
 	
 #include "CoreMinimal.h"
 
-#if WINDOWS_USE_FEATURE_LAUNCH
-
 #include "Misc/App.h"
 #include "Misc/OutputDeviceError.h"
 #include "LaunchEngineLoop.h"
@@ -286,5 +284,3 @@ int32 WINAPI WinMain(_In_ HINSTANCE hInInstance, _In_opt_ HINSTANCE hPrevInstanc
 	return Result;
 }
 
-#endif //WINDOWS_USE_FEATURE_LAUNCH
-
diff --git a/Engine/Source/Runtime/UELibrary/Private/UELibrary.cpp b/Engine/Source/Runtime/UELibrary/Private/UELibrary.cpp
--- a/Engine/Source/Runtime/UELibrary/Private/UELibrary.cpp
+++ b/Engine/Source/Runtime/UELibrary/Private/UELibrary.cpp
@@ -128,8 +128,6 @@ int UELibrary_Tick()
 	return 0;
 }
 
-#if WINDOWS_USE_FEATURE_LAUNCH
-
 LRESULT UELibrary_WndProc(HWND hWnd, UINT message, WPARAM wParam, LPARAM lParam)
 {
 	// Do default handling if the library isn't initialized
@@ -143,8 +141,6 @@ LRESULT UELibrary_WndProc(HWND hWnd, UINT message, WPARAM wParam, LPARAM lParam)
 	}
 }
 
-#endif
-
 int UELibrary_Shutdown()
 {
 	// Validate that we have already been initialized
